- Removing "EditAnywhere" from ULuminRuntimeSettings::bUseVulkan (and category because it errors if you have a category on something that is not EditAnywhere)
- Defaulting [/Script/LuminRuntimeSettings.LuminRuntimeSettings] bUseVulkan to true

Tested that bUseVulkan is true in a new project, and that the setting is now hidden from Project Settings > Magic Leap > Build

diff --git a/Engine/Config/BaseEngine.ini b/Engine/Config/BaseEngine.ini
--- a/Engine/Config/BaseEngine.ini
+++ b/Engine/Config/BaseEngine.ini
@@ -1773,7 +1773,7 @@ bBuildWithNvTegraGfxDebugger=False
 Certificate=(FilePath="")
 IconModelPath=(Path="../../Build/Lumin/Resources/Model")
 IconPortalPath=(Path="../../Build/Lumin/Resources/Portal")
-bUseVulkan=false
+bUseVulkan=true
 VersionCode=1
 bInternetRequired=true
 +AppPrivileges=Headpose
diff --git a/Engine/Source/Runtime/Lumin/LuminRuntimeSettings/Classes/LuminRuntimeSettings.h b/Engine/Source/Runtime/Lumin/LuminRuntimeSettings/Classes/LuminRuntimeSettings.h
--- a/Engine/Source/Runtime/Lumin/LuminRuntimeSettings/Classes/LuminRuntimeSettings.h
+++ b/Engine/Source/Runtime/Lumin/LuminRuntimeSettings/Classes/LuminRuntimeSettings.h
@@ -58,11 +58,11 @@ public:
 	UPROPERTY(GlobalConfig, EditAnywhere, Category = "Runtime", Meta = (DisplayName = "Protected Content"))
 	bool bProtectedContent;
 
-	/** Use ES2/Vulkan Mobile Rendering. */
+	/** Use Vulkan Mobile Rendering. */
 	UPROPERTY(GlobalConfig, EditAnywhere, Category = "Build", Meta = (DisplayName = "Use Mobile Rendering"))
 	bool bUseMobileRendering;
 
-	UPROPERTY(GlobalConfig, EditAnywhere, Category = "Build", Meta = (DisplayName = "Use Vulkan (otherwise, OpenGL)"))
+	UPROPERTY(GlobalConfig, Meta = (DisplayName = "Use Vulkan (otherwise, OpenGL)"))
 	bool bUseVulkan;
 
 	/** Enable support for NVIDIA Tegra Graphics Debugger? [FOR FULL SOURCE GAMES ONLY]. */
