Magic Leap settings: Changing comment/tooltip on mobile rendering flag to not mention vulkan and just say "mobile". Still technically possible to use ES2 but it's hidden.

diff --git a/Engine/Source/Runtime/Lumin/LuminRuntimeSettings/Classes/LuminRuntimeSettings.h b/Engine/Source/Runtime/Lumin/LuminRuntimeSettings/Classes/LuminRuntimeSettings.h
--- a/Engine/Source/Runtime/Lumin/LuminRuntimeSettings/Classes/LuminRuntimeSettings.h
+++ b/Engine/Source/Runtime/Lumin/LuminRuntimeSettings/Classes/LuminRuntimeSettings.h
@@ -58,7 +58,7 @@ public:
 	UPROPERTY(GlobalConfig, EditAnywhere, Category = "Runtime", Meta = (DisplayName = "Protected Content"))
 	bool bProtectedContent;
 
-	/** Use Vulkan Mobile Rendering. */
+	/** If checked, use Mobile Rendering. Otherwise, use Desktop Rendering. */
 	UPROPERTY(GlobalConfig, EditAnywhere, Category = "Build", Meta = (DisplayName = "Use Mobile Rendering"))
 	bool bUseMobileRendering;
 
