#DEV-VR - Removing question mark from tooltip

diff --git a/Engine/Source/Runtime/Lumin/LuminRuntimeSettings/Classes/LuminRuntimeSettings.h b/Engine/Source/Runtime/Lumin/LuminRuntimeSettings/Classes/LuminRuntimeSettings.h
--- a/Engine/Source/Runtime/Lumin/LuminRuntimeSettings/Classes/LuminRuntimeSettings.h
+++ b/Engine/Source/Runtime/Lumin/LuminRuntimeSettings/Classes/LuminRuntimeSettings.h
@@ -65,7 +65,7 @@ public:
 	UPROPERTY(GlobalConfig, EditAnywhere, Category = "Build", Meta = (DisplayName = "Use Vulkan (otherwise, OpenGL)"))
 	bool bUseVulkan;
 
-	/** Enable support for NVIDIA Tegra Graphics Debugger? [FOR FULL SOURCE GAMES ONLY]. */
+	/** Enable support for NVIDIA Tegra Graphics Debugger [FOR FULL SOURCE GAMES ONLY]. */
 	UPROPERTY(GlobalConfig, EditAnywhere, Category = "Build", Meta = (DisplayName = "Support NVIDIA Tegra Graphics Debugger"))
 	bool bBuildWithNvTegraGfxDebugger;
 
