#DEV-VR - Disable editing of Tegra Debugger setting for installed builds

diff --git a/Engine/Source/Runtime/Lumin/LuminRuntimeSettings/Classes/LuminRuntimeSettings.h b/Engine/Source/Runtime/Lumin/LuminRuntimeSettings/Classes/LuminRuntimeSettings.h
--- a/Engine/Source/Runtime/Lumin/LuminRuntimeSettings/Classes/LuminRuntimeSettings.h
+++ b/Engine/Source/Runtime/Lumin/LuminRuntimeSettings/Classes/LuminRuntimeSettings.h
@@ -35,6 +35,8 @@ public:
 	GENERATED_BODY()
 
 #if WITH_EDITOR
+	virtual bool CanEditChange(const UProperty* InProperty) const override;
+
 	virtual void PostEditChangeProperty(struct FPropertyChangedEvent& PropertyChangedEvent) override;
 #endif // WITH_EDITOR
 
diff --git a/Engine/Source/Runtime/Lumin/LuminRuntimeSettings/Private/LuminRuntimeSettings.cpp b/Engine/Source/Runtime/Lumin/LuminRuntimeSettings/Private/LuminRuntimeSettings.cpp
--- a/Engine/Source/Runtime/Lumin/LuminRuntimeSettings/Private/LuminRuntimeSettings.cpp
+++ b/Engine/Source/Runtime/Lumin/LuminRuntimeSettings/Private/LuminRuntimeSettings.cpp
@@ -3,9 +3,22 @@
 #include "LuminRuntimeSettings.h"
 #include "Misc/ConfigCacheIni.h"
 #include "CoreGlobals.h"
+#include "Misc/App.h"
 
 #if WITH_EDITOR
 
+bool ULuminRuntimeSettings::CanEditChange(const UProperty* InProperty) const
+{
+	const bool ParentVal = Super::CanEditChange(InProperty);
+
+	if ((InProperty->GetFName() == GET_MEMBER_NAME_CHECKED(ULuminRuntimeSettings, bBuildWithNvTegraGfxDebugger)))
+	{
+		return !FApp::IsEngineInstalled();
+	}
+
+	return ParentVal;
+}
+
 void ULuminRuntimeSettings::PostEditChangeProperty(struct FPropertyChangedEvent& PropertyChangedEvent)
 {
 	GConfig->Flush(1);
