#DEV-VR - Removing ES2 from bMobileRendering Tooltop

diff --git a/Engine/Source/Runtime/Lumin/LuminRuntimeSettings/Classes/LuminRuntimeSettings.h b/Engine/Source/Runtime/Lumin/LuminRuntimeSettings/Classes/LuminRuntimeSettings.h
--- a/Engine/Source/Runtime/Lumin/LuminRuntimeSettings/Classes/LuminRuntimeSettings.h
+++ b/Engine/Source/Runtime/Lumin/LuminRuntimeSettings/Classes/LuminRuntimeSettings.h
@@ -60,7 +60,7 @@ public:
 	UPROPERTY(GlobalConfig, EditAnywhere, Category = "Runtime", Meta = (DisplayName = "Protected Content"))
 	bool bProtectedContent;
 
-	/** Use ES2/Vulkan Mobile Rendering. */
+	/** If checked, use Vulkan Mobile Rendering. Otherwise, use Vulkan Desktop Rendering. */
 	UPROPERTY(GlobalConfig, EditAnywhere, Category = "Build", Meta = (DisplayName = "Use Mobile Rendering"))
 	bool bUseMobileRendering;
 
