Merging CL-4118965 from Release-4.20 using DevVRtoRelease420

diff --git a/Engine/Config/BaseEngine.ini b/Engine/Config/BaseEngine.ini
--- a/Engine/Config/BaseEngine.ini
+++ b/Engine/Config/BaseEngine.ini
@@ -1775,7 +1775,7 @@ bBuildWithNvTegraGfxDebugger=False
 Certificate=(FilePath="")
 IconModelPath=(Path="../../Build/Lumin/Resources/Model")
 IconPortalPath=(Path="../../Build/Lumin/Resources/Portal")
-bUseVulkan=false
+bUseVulkan=true
 VersionCode=1
 bInternetRequired=true
 +AppPrivileges=Headpose
diff --git a/Engine/Source/Runtime/Lumin/LuminRuntimeSettings/Classes/LuminRuntimeSettings.h b/Engine/Source/Runtime/Lumin/LuminRuntimeSettings/Classes/LuminRuntimeSettings.h
--- a/Engine/Source/Runtime/Lumin/LuminRuntimeSettings/Classes/LuminRuntimeSettings.h
+++ b/Engine/Source/Runtime/Lumin/LuminRuntimeSettings/Classes/LuminRuntimeSettings.h
@@ -64,7 +64,7 @@ public:
 	UPROPERTY(GlobalConfig, EditAnywhere, Category = "Build", Meta = (DisplayName = "Use Mobile Rendering"))
 	bool bUseMobileRendering;
 
-	UPROPERTY(GlobalConfig, EditAnywhere, Category = "Build", Meta = (DisplayName = "Use Vulkan (otherwise, OpenGL)"))
+	UPROPERTY(GlobalConfig, Meta = (DisplayName = "Use Vulkan (otherwise, OpenGL)"))
 	bool bUseVulkan;
 
 	/** Enable support for NVIDIA Tegra Graphics Debugger [FOR FULL SOURCE GAMES ONLY]. */
