#DEV-VR - Making bUseMobileRendering non-editable for binary builds

diff --git a/Engine/Source/Runtime/Lumin/LuminRuntimeSettings/Classes/LuminRuntimeSettings.h b/Engine/Source/Runtime/Lumin/LuminRuntimeSettings/Classes/LuminRuntimeSettings.h
--- a/Engine/Source/Runtime/Lumin/LuminRuntimeSettings/Classes/LuminRuntimeSettings.h
+++ b/Engine/Source/Runtime/Lumin/LuminRuntimeSettings/Classes/LuminRuntimeSettings.h
@@ -95,8 +95,8 @@ public:
 	UPROPERTY(GlobalConfig, EditAnywhere, Category = "Runtime", Meta = (DisplayName = "Protected Content"))
 	bool bProtectedContent;
 
-	/** If checked, use Mobile Rendering. Otherwise, use Desktop Rendering. */
-	UPROPERTY(GlobalConfig, EditAnywhere, Category = "Build", Meta = (DisplayName = "Use Mobile Rendering (otherwise, Desktop Rendering)"))
+	/** If checked, use Mobile Rendering. Otherwise, use Desktop Rendering [FOR FULL SOURCE GAMES ONLY]. */
+	UPROPERTY(GlobalConfig, EditAnywhere, Category = "Build", Meta = (DisplayName = "Use Mobile Rendering.  Otherwise, Desktop Rendering"))
 	bool bUseMobileRendering;
 
 	UPROPERTY(GlobalConfig, Meta = (DisplayName = "Use Vulkan (otherwise, OpenGL)"))
diff --git a/Engine/Source/Runtime/Lumin/LuminRuntimeSettings/Private/LuminRuntimeSettings.cpp b/Engine/Source/Runtime/Lumin/LuminRuntimeSettings/Private/LuminRuntimeSettings.cpp
--- a/Engine/Source/Runtime/Lumin/LuminRuntimeSettings/Private/LuminRuntimeSettings.cpp
+++ b/Engine/Source/Runtime/Lumin/LuminRuntimeSettings/Private/LuminRuntimeSettings.cpp
@@ -11,7 +11,8 @@ bool ULuminRuntimeSettings::CanEditChange(const UProperty* InProperty) const
 {
 	const bool ParentVal = Super::CanEditChange(InProperty);
 
-	if ((InProperty->GetFName() == GET_MEMBER_NAME_CHECKED(ULuminRuntimeSettings, bBuildWithNvTegraGfxDebugger)))
+	if ((InProperty->GetFName() == GET_MEMBER_NAME_CHECKED(ULuminRuntimeSettings, bBuildWithNvTegraGfxDebugger)) ||
+		(InProperty->GetFName() == GET_MEMBER_NAME_CHECKED(ULuminRuntimeSettings, bUseMobileRendering)))
 	{
 		return !FApp::IsEngineInstalled();
 	}
